Sequencer: Fixed custom BP master tracks not working correctly

diff --git a/Engine/Plugins/MovieScene/CustomizableSequencerTracks/Source/CustomizableSequencerTracks/Private/SequencerSectionBP.cpp b/Engine/Plugins/MovieScene/CustomizableSequencerTracks/Source/CustomizableSequencerTracks/Private/SequencerSectionBP.cpp
--- a/Engine/Plugins/MovieScene/CustomizableSequencerTracks/Source/CustomizableSequencerTracks/Private/SequencerSectionBP.cpp
+++ b/Engine/Plugins/MovieScene/CustomizableSequencerTracks/Source/CustomizableSequencerTracks/Private/SequencerSectionBP.cpp
@@ -34,6 +34,7 @@ void USequencerSectionBP::ImportEntityImpl(UMovieSceneEntitySystemLinker* Entity
 			FEntityBuilder()
 			.Add(BuiltInComponents->TrackInstance, FMovieSceneTrackInstanceComponent{ decltype(FMovieSceneTrackInstanceComponent::Owner)(this), CustomTrack->TrackInstanceType })
 			.AddConditional(BuiltInComponents->GenericObjectBinding, ObjectBindingID, ObjectBindingID.IsValid())
+			.AddTagConditional(BuiltInComponents->Tags.Master, !ObjectBindingID.IsValid())
 		);
 	}
 }
\ No newline at end of file
