Sequencer: Add missing Track->Modify() when adding a custom section

diff --git a/Engine/Plugins/MovieScene/CustomizableSequencerTracks/Source/CustomizableSequencerTracksEditor/Private/SequencerTrackBPEditor.cpp b/Engine/Plugins/MovieScene/CustomizableSequencerTracks/Source/CustomizableSequencerTracksEditor/Private/SequencerTrackBPEditor.cpp
--- a/Engine/Plugins/MovieScene/CustomizableSequencerTracks/Source/CustomizableSequencerTracksEditor/Private/SequencerTrackBPEditor.cpp
+++ b/Engine/Plugins/MovieScene/CustomizableSequencerTracks/Source/CustomizableSequencerTracksEditor/Private/SequencerTrackBPEditor.cpp
@@ -71,7 +71,7 @@ namespace
 		UObject* LoadedAsset = AssetData.FastGetAsset();
 		if (!LoadedAsset)
 		{
-			FScopedSlowTask SlowTask(1.f, FText::Format(LOCTEXT("LoadingClass", "Loading asset %s"), FText::FromName(AssetData.AssetName)));
+			FScopedSlowTask SlowTask(1.f, FText::Format(LOCTEXT("LoadingClass", "Loading asset {0}"), FText::FromName(AssetData.AssetName)));
 			SlowTask.MakeDialogDelayed(1.f);
 
 			LoadedAsset = AssetData.GetAsset();
@@ -310,7 +310,10 @@ void FSequencerTrackBPEditor::CreateNewSection(USequencerTrackBP* Track, TSubcla
 
 	if (Class && SequencerPin)
 	{
-		FScopedTransaction Transaction(FText::Format(LOCTEXT("AddCustomSection_Transaction", "Add New Section From Class %s"), FText::FromName(Class->GetFName())));
+		FScopedTransaction Transaction(FText::Format(LOCTEXT("AddCustomSection_Transaction", "Add New Section From Class {0}"), FText::FromName(Class->GetFName())));
+
+		Track->Modify();
+			
 		USequencerSectionBP* NewSection = NewObject<USequencerSectionBP>(Track, Class, NAME_None, RF_Transactional);
 
 		FQualifiedFrameTime CurrentTime = SequencerPin->GetLocalTime();
@@ -345,7 +348,7 @@ void FSequencerTrackBPEditor::AddNewMasterTrack(FAssetData AssetData)
 		return;
 	}
 
-	const FScopedTransaction Transaction(FText::Format(LOCTEXT("AddCustomMasterTrack_Transaction", "Add Master Track %s"), FText::FromName(ClassToAdd->GetFName())));
+	const FScopedTransaction Transaction(FText::Format(LOCTEXT("AddCustomMasterTrack_Transaction", "Add Master Track {0}"), FText::FromName(ClassToAdd->GetFName())));
 
 	MovieScene->Modify();
 
@@ -373,7 +376,7 @@ void FSequencerTrackBPEditor::AddNewObjectBindingTrack(FAssetData AssetData, TAr
 		return;
 	}
 
-	const FScopedTransaction Transaction(FText::Format(LOCTEXT("AddCustomObjectTrack_Transaction", "Add Object Track %s"), FText::FromName(ClassToAdd->GetFName())));
+	const FScopedTransaction Transaction(FText::Format(LOCTEXT("AddCustomObjectTrack_Transaction", "Add Object Track {0}"), FText::FromName(ClassToAdd->GetFName())));
 
 	MovieScene->Modify();
 
