Sequencer: OnInputRemoved for any left over inputs OnDestroyed()

diff --git a/Engine/Plugins/MovieScene/CustomizableSequencerTracks/Source/CustomizableSequencerTracks/Private/SequencerTrackInstanceBP.cpp b/Engine/Plugins/MovieScene/CustomizableSequencerTracks/Source/CustomizableSequencerTracks/Private/SequencerTrackInstanceBP.cpp
--- a/Engine/Plugins/MovieScene/CustomizableSequencerTracks/Source/CustomizableSequencerTracks/Private/SequencerTrackInstanceBP.cpp
+++ b/Engine/Plugins/MovieScene/CustomizableSequencerTracks/Source/CustomizableSequencerTracks/Private/SequencerTrackInstanceBP.cpp
@@ -103,4 +103,26 @@ void USequencerTrackInstanceBP::OnInputRemoved(const FMovieSceneTrackInstanceInp
 }
 
 
+void USequencerTrackInstanceBP::OnDestroyed()
+{
+	using namespace UE::MovieScene;
+
+	TArrayView<const FMovieSceneTrackInstanceInput> InputsView = Super::GetInputs();
+	for (const FMovieSceneTrackInstanceInput& Input : InputsView)
+	{
+		FSequencerTrackInstanceInput ThisInput;
+		ThisInput.Section = CastChecked<USequencerSectionBP>(Input.Section);
+
+		const FInstanceRegistry* InstanceRegistry = GetLinker()->GetInstanceRegistry();
+		if (ensure(InstanceRegistry->IsHandleValid(Input.InstanceHandle)))
+		{
+			ThisInput.Context = InstanceRegistry->GetInstance(Input.InstanceHandle).GetContext();
+		}
+
+		K2_OnInputRemoved(ThisInput);
+	}
+	K2_OnDestroyed();
+}
+
+
 IMPLEMENT_MODULE(FDefaultModuleImpl, CustomizableSequencerTracks);
\ No newline at end of file
diff --git a/Engine/Plugins/MovieScene/CustomizableSequencerTracks/Source/CustomizableSequencerTracks/Public/SequencerTrackInstanceBP.h b/Engine/Plugins/MovieScene/CustomizableSequencerTracks/Source/CustomizableSequencerTracks/Public/SequencerTrackInstanceBP.h
--- a/Engine/Plugins/MovieScene/CustomizableSequencerTracks/Source/CustomizableSequencerTracks/Public/SequencerTrackInstanceBP.h
+++ b/Engine/Plugins/MovieScene/CustomizableSequencerTracks/Source/CustomizableSequencerTracks/Public/SequencerTrackInstanceBP.h
@@ -85,6 +85,6 @@ private:
 
 	virtual void OnEndUpdateInputs() override { K2_OnEndUpdateInputs(); }
 
-	virtual void OnDestroyed() override { K2_OnDestroyed(); }
+	virtual void OnDestroyed() override;
 };
 
